UE4 - Metal shader format

diff --git a/Engine/Source/Developer/MetalShaderFormat/MetalShaderFormat.Build.cs b/Engine/Source/Developer/MetalShaderFormat/MetalShaderFormat.Build.cs
new file mode 100644
--- /dev/null
+++ b/Engine/Source/Developer/MetalShaderFormat/MetalShaderFormat.Build.cs
@@ -0,0 +1,23 @@
+// Copyright 1998-2014 Epic Games, Inc. All Rights Reserved.
+
+using UnrealBuildTool;
+
+public class MetalShaderFormat : ModuleRules
+{
+	public MetalShaderFormat(TargetInfo Target)
+	{
+		PrivateIncludePathModuleNames.Add("TargetPlatform");
+
+		PrivateDependencyModuleNames.AddRange(
+			new string[] {
+				"Core",
+				"ShaderCore",
+				"ShaderPreprocessor"
+			}
+			);
+
+		AddThirdPartyPrivateStaticDependencies(Target, 
+			"HLSLCC"
+			);
+	}
+}
diff --git a/Engine/Source/Developer/MetalShaderFormat/Private/MetalShaderCompiler.cpp b/Engine/Source/Developer/MetalShaderFormat/Private/MetalShaderCompiler.cpp
new file mode 100644
--- /dev/null
+++ b/Engine/Source/Developer/MetalShaderFormat/Private/MetalShaderCompiler.cpp
@@ -0,0 +1,917 @@
+// Copyright 1998-2014 Epic Games, Inc. All Rights Reserved.
+///
+
+#include "UrbanShaderFormat.h"
+#include "Core.h"
+#include "../../../EmptyRHI/Public/EmptyShaderResources.h"
+
+
+#if PLATFORM_WINDOWS
+#include "AllowWindowsPlatformTypes.h"
+	#include "Windows/PreWindowsApi.h"
+	#include <objbase.h>
+	#include <assert.h>
+	#include <stdio.h>
+	#include "Windows/PostWindowsApi.h"
+	#include "Windows/MinWindows.h"
+#include "HideWindowsPlatformTypes.h"
+#endif
+#include "ShaderPreprocessor.h"
+#include "hlslcc.h"
+#include "Metal/MetalBackend.h"
+
+DEFINE_LOG_CATEGORY_STATIC(LogUrbanShaderCompiler, Log, All); 
+
+/*------------------------------------------------------------------------------
+	Shader compiling.
+------------------------------------------------------------------------------*/
+
+/** Map shader frequency -> string for messages. */
+static const TCHAR* GLFrequencyStringTable[] =
+{
+	TEXT("Vertex"),
+	TEXT("Hull"),
+	TEXT("Domain"),
+	TEXT("Pixel"),
+	TEXT("Geometry"),
+	TEXT("Compute")
+};
+
+static FString ParseIdentifier(const ANSICHAR* &Str)
+{
+	FString Result;
+
+	while ((*Str >= 'A' && *Str <= 'Z')
+		|| (*Str >= 'a' && *Str <= 'z')
+		|| (*Str >= '0' && *Str <= '9')
+		|| *Str == '_')
+	{
+		Result += *Str;
+		++Str;
+	}
+
+	return Result;
+}
+
+static bool Match(const ANSICHAR* &Str, ANSICHAR Char)
+{
+	if (*Str == Char)
+	{
+		++Str;
+		return true;
+	}
+
+	return false;
+}
+
+static uint32 ParseNumber(const ANSICHAR* &Str)
+{
+	uint32 Num = 0;
+	while (*Str && *Str >= '0' && *Str <= '9')
+	{
+		Num = Num * 10 + *Str++ - '0';
+	}
+	return Num;
+}
+
+/**
+ * Construct the final microcode from the compiled and verified shader source.
+ * @param ShaderOutput - Where to store the microcode and parameter map.
+ * @param InShaderSource - Metal source with input/output signature.
+ * @param SourceLen - The length of the Metal source code.
+ */
+static void BuildMetalShaderOutput(
+	FShaderCompilerOutput& ShaderOutput,
+	const ANSICHAR* InShaderSource,
+	int32 SourceLen,
+	TArray<FShaderCompilerError>& OutErrors
+	)
+{
+	FMetalCodeHeader Header = {0};
+	const ANSICHAR* ShaderSource = InShaderSource;
+	FShaderParameterMap& ParameterMap = ShaderOutput.ParameterMap;
+	EShaderFrequency Frequency = (EShaderFrequency)ShaderOutput.Target.Frequency;
+
+	// Write out the magic markers.
+	Header.Frequency = Frequency;
+
+	#define DEF_PREFIX_STR(Str) \
+		const ANSICHAR* Str##Prefix = "// @" #Str ": "; \
+		const int32 Str##PrefixLen = FCStringAnsi::Strlen(Str##Prefix)
+	DEF_PREFIX_STR(Inputs);
+	DEF_PREFIX_STR(Outputs);
+	DEF_PREFIX_STR(UniformBlocks);
+	DEF_PREFIX_STR(Uniforms);
+	DEF_PREFIX_STR(PackedGlobals);
+	DEF_PREFIX_STR(PackedUB);
+	DEF_PREFIX_STR(PackedUBCopies);
+	DEF_PREFIX_STR(PackedUBGlobalCopies);
+	DEF_PREFIX_STR(Samplers);
+	DEF_PREFIX_STR(UAVs);
+	DEF_PREFIX_STR(SamplerStates);
+	#undef DEF_PREFIX_STR
+
+	// Skip any comments that come before the signature.
+	while (	FCStringAnsi::Strncmp(ShaderSource, "//", 2) == 0 &&
+			FCStringAnsi::Strncmp(ShaderSource, "// @", 4) != 0 )
+	{
+		while (*ShaderSource && *ShaderSource++ != '\n') {}
+	}
+
+	// HLSLCC first prints the list of inputs.
+	if (FCStringAnsi::Strncmp(ShaderSource, InputsPrefix, InputsPrefixLen) == 0)
+	{
+		ShaderSource += InputsPrefixLen;
+
+		// Only inputs for vertex shaders must be tracked.
+		if (Frequency == SF_Vertex)
+		{
+			const ANSICHAR* AttributePrefix = "in_ATTRIBUTE";
+			const int32 AttributePrefixLen = FCStringAnsi::Strlen(AttributePrefix);
+			while (*ShaderSource && *ShaderSource != '\n')
+			{
+				// Skip the type.
+				while (*ShaderSource && *ShaderSource++ != ':') {}
+				
+				// Only process attributes.
+				if (FCStringAnsi::Strncmp(ShaderSource, AttributePrefix, AttributePrefixLen) == 0)
+				{
+					ShaderSource += AttributePrefixLen;
+					uint8 AttributeIndex = ParseNumber(ShaderSource);
+					Header.Bindings.InOutMask |= (1 << AttributeIndex);
+				}
+
+				// Skip to the next.
+				while (*ShaderSource && *ShaderSource != ',' && *ShaderSource != '\n')
+				{
+					ShaderSource++;
+				}
+
+				if (Match(ShaderSource, '\n'))
+				{
+					break;
+				}
+
+				verify(Match(ShaderSource, ','));
+			}
+		}
+		else
+		{
+			// Skip to the next line.
+			while (*ShaderSource && *ShaderSource++ != '\n') {}
+		}
+	}
+
+	// Then the list of outputs.
+	if (FCStringAnsi::Strncmp(ShaderSource, OutputsPrefix, OutputsPrefixLen) == 0)
+	{
+		ShaderSource += OutputsPrefixLen;
+
+		// Only outputs for pixel shaders must be tracked.
+		if (Frequency == SF_Pixel)
+		{
+			const ANSICHAR* TargetPrefix = "out_Target";
+			const int32 TargetPrefixLen = FCStringAnsi::Strlen(TargetPrefix);
+
+			while (*ShaderSource && *ShaderSource != '\n')
+			{
+				// Skip the type.
+				while (*ShaderSource && *ShaderSource++ != ':') {}
+
+				// Handle targets.
+				if (FCStringAnsi::Strncmp(ShaderSource, TargetPrefix, TargetPrefixLen) == 0)
+				{
+					ShaderSource += TargetPrefixLen;
+					uint8 TargetIndex = ParseNumber(ShaderSource);
+					Header.Bindings.InOutMask |= (1 << TargetIndex);
+				}
+				// Handle depth writes.
+				else if (FCStringAnsi::Strcmp(ShaderSource, "gl_FragDepth") == 0)
+				{
+					Header.Bindings.InOutMask |= 0x8000;
+				}
+
+				// Skip to the next.
+				while (*ShaderSource && *ShaderSource != ',' && *ShaderSource != '\n')
+				{
+					ShaderSource++;
+				}
+
+				if (Match(ShaderSource, '\n'))
+				{
+					break;
+				}
+
+				verify(Match(ShaderSource, ','));
+			}
+		}
+		else
+		{
+			// Skip to the next line.
+			while (*ShaderSource && *ShaderSource++ != '\n') {}
+		}
+	}
+
+	bool bHasRegularUniformBuffers = false;
+
+	// Then 'normal' uniform buffers.
+	if (FCStringAnsi::Strncmp(ShaderSource, UniformBlocksPrefix, UniformBlocksPrefixLen) == 0)
+	{
+		ShaderSource += UniformBlocksPrefixLen;
+
+		while (*ShaderSource && *ShaderSource != '\n')
+		{
+			FString BufferName = ParseIdentifier(ShaderSource);
+			verify(BufferName.Len() > 0);
+			verify(Match(ShaderSource, '('));
+			uint16 UBIndex = ParseNumber(ShaderSource);
+			if (UBIndex >= Header.Bindings.NumUniformBuffers)
+			{
+				Header.Bindings.NumUniformBuffers = UBIndex + 1;
+			}
+			verify(Match(ShaderSource, ')'));
+			ParameterMap.AddParameterAllocation(*BufferName, UBIndex, 0, 0);
+			bHasRegularUniformBuffers = true;
+
+			// Skip the comma.
+			if (Match(ShaderSource, '\n'))
+			{
+				break;
+			}
+
+			verify(Match(ShaderSource, ','));
+		}
+
+		Match(ShaderSource, '\n');
+	}
+
+	// Then uniforms.
+	const uint16 BytesPerComponent = 4;
+/*
+	uint16 PackedUniformSize[OGL_NUM_PACKED_UNIFORM_ARRAYS] = {0};
+	FMemory::Memzero(&PackedUniformSize, sizeof(PackedUniformSize));
+*/
+	if (FCStringAnsi::Strncmp(ShaderSource, UniformsPrefix, UniformsPrefixLen) == 0)
+	{
+		// @todo-mobile: Will we ever need to support this code path?
+		check(0);
+/*
+		ShaderSource += UniformsPrefixLen;
+
+		while (*ShaderSource && *ShaderSource != '\n')
+		{
+			uint16 ArrayIndex = 0;
+			uint16 Offset = 0;
+			uint16 NumComponents = 0;
+
+			FString ParameterName = ParseIdentifier(ShaderSource);
+			verify(ParameterName.Len() > 0);
+			verify(Match(ShaderSource, '('));
+			ArrayIndex = ParseNumber(ShaderSource);
+			verify(Match(ShaderSource, ':'));
+			Offset = ParseNumber(ShaderSource);
+			verify(Match(ShaderSource, ':'));
+			NumComponents = ParseNumber(ShaderSource);
+			verify(Match(ShaderSource, ')'));
+
+			ParameterMap.AddParameterAllocation(
+				*ParameterName,
+				ArrayIndex,
+				Offset * BytesPerComponent,
+				NumComponents * BytesPerComponent
+				);
+
+			if (ArrayIndex < OGL_NUM_PACKED_UNIFORM_ARRAYS)
+			{
+				PackedUniformSize[ArrayIndex] = FMath::Max<uint16>(
+					PackedUniformSize[ArrayIndex],
+					BytesPerComponent * (Offset + NumComponents)
+					);
+			}
+
+			// Skip the comma.
+			if (Match(ShaderSource, '\n'))
+			{
+				break;
+			}
+
+			verify(Match(ShaderSource, ','));
+		}
+
+		Match(ShaderSource, '\n');
+*/
+	}
+
+	// Packed global uniforms
+	TMap<ANSICHAR, uint16> PackedGlobalArraySize;
+	if (FCStringAnsi::Strncmp(ShaderSource, PackedGlobalsPrefix, PackedGlobalsPrefixLen) == 0)
+	{
+		ShaderSource += PackedGlobalsPrefixLen;
+		while (*ShaderSource && *ShaderSource != '\n')
+		{
+			ANSICHAR ArrayIndex = 0;
+			uint16 Offset = 0;
+			uint16 NumComponents = 0;
+
+			FString ParameterName = ParseIdentifier(ShaderSource);
+			verify(ParameterName.Len() > 0);
+			verify(Match(ShaderSource, '('));
+			ArrayIndex = *ShaderSource++;
+			verify(Match(ShaderSource, ':'));
+			Offset = ParseNumber(ShaderSource);
+			verify(Match(ShaderSource, ','));
+			NumComponents = ParseNumber(ShaderSource);
+			verify(Match(ShaderSource, ')'));
+
+			ParameterMap.AddParameterAllocation(
+				*ParameterName,
+				ArrayIndex,
+				Offset * BytesPerComponent,
+				NumComponents * BytesPerComponent
+				);
+
+			uint16& Size = PackedGlobalArraySize.FindOrAdd(ArrayIndex);
+			Size = FMath::Max<uint16>(BytesPerComponent * (Offset + NumComponents), Size);
+
+			if (Match(ShaderSource, '\n'))
+			{
+				break;
+			}
+
+			// Skip the comma.
+			verify(Match(ShaderSource, ','));
+		}
+
+		Match(ShaderSource, '\n');
+	}
+
+	// Packed Uniform Buffers
+	TMap<int, TMap<ANSICHAR, uint16> > PackedUniformBuffersSize;
+	while (FCStringAnsi::Strncmp(ShaderSource, PackedUBPrefix, PackedUBPrefixLen) == 0)
+	{
+		ShaderSource += PackedUBPrefixLen;
+		FString BufferName = ParseIdentifier(ShaderSource);
+		verify(BufferName.Len() > 0);
+		verify(Match(ShaderSource, '('));
+		uint16 BufferIndex = ParseNumber(ShaderSource);
+		check(BufferIndex == Header.Bindings.NumUniformBuffers);
+		verify(Match(ShaderSource, ')'));
+		ParameterMap.AddParameterAllocation(*BufferName, Header.Bindings.NumUniformBuffers++, 0, 0);
+
+		verify(Match(ShaderSource, ':'));
+		Match(ShaderSource, ' ');
+		while (*ShaderSource && *ShaderSource != '\n')
+		{
+			FString ParameterName = ParseIdentifier(ShaderSource);
+			verify(ParameterName.Len() > 0);
+			verify(Match(ShaderSource, '('));
+			ParseNumber(ShaderSource);
+			verify(Match(ShaderSource, ','));
+			ParseNumber(ShaderSource);
+			verify(Match(ShaderSource, ')'));
+
+			if (Match(ShaderSource, '\n'))
+			{
+				break;
+			}
+
+			verify(Match(ShaderSource, ','));
+		}
+	}
+
+	// Packed Uniform Buffers copy lists & setup sizes for each UB/Precision entry
+	if (FCStringAnsi::Strncmp(ShaderSource, PackedUBCopiesPrefix, PackedUBCopiesPrefixLen) == 0)
+	{
+		ShaderSource += PackedUBCopiesPrefixLen;
+		while (*ShaderSource && *ShaderSource != '\n')
+		{
+			FMetalUniformBufferCopyInfo CopyInfo;
+
+			CopyInfo.SourceUBIndex = ParseNumber(ShaderSource);
+			verify(Match(ShaderSource, ':'));
+
+			CopyInfo.SourceOffsetInFloats = ParseNumber(ShaderSource);
+			verify(Match(ShaderSource, '-'));
+
+			CopyInfo.DestUBIndex = ParseNumber(ShaderSource);
+			verify(Match(ShaderSource, ':'));
+
+			CopyInfo.DestUBTypeName = *ShaderSource++;
+			CopyInfo.DestUBTypeIndex = MetalPackedTypeNameToTypeIndex(CopyInfo.DestUBTypeName);
+			verify(Match(ShaderSource, ':'));
+
+			CopyInfo.DestOffsetInFloats = ParseNumber(ShaderSource);
+			verify(Match(ShaderSource, ':'));
+
+			CopyInfo.SizeInFloats = ParseNumber(ShaderSource);
+
+			Header.UniformBuffersCopyInfo.Add(CopyInfo);
+
+			auto& UniformBufferSize = PackedUniformBuffersSize.FindOrAdd(CopyInfo.DestUBIndex);
+			uint16& Size = UniformBufferSize.FindOrAdd(CopyInfo.DestUBTypeName);
+			Size = FMath::Max<uint16>(BytesPerComponent * (CopyInfo.DestOffsetInFloats + CopyInfo.SizeInFloats), Size);
+
+			if (Match(ShaderSource, '\n'))
+			{
+				break;
+			}
+
+			verify(Match(ShaderSource, ','));
+		}
+	}
+
+	if (FCStringAnsi::Strncmp(ShaderSource, PackedUBGlobalCopiesPrefix, PackedUBGlobalCopiesPrefixLen) == 0)
+	{
+		ShaderSource += PackedUBGlobalCopiesPrefixLen;
+		while (*ShaderSource && *ShaderSource != '\n')
+		{
+			FMetalUniformBufferCopyInfo CopyInfo;
+
+			CopyInfo.SourceUBIndex = ParseNumber(ShaderSource);
+			verify(Match(ShaderSource, ':'));
+
+			CopyInfo.SourceOffsetInFloats = ParseNumber(ShaderSource);
+			verify(Match(ShaderSource, '-'));
+
+			CopyInfo.DestUBIndex = 0;
+
+			CopyInfo.DestUBTypeName = *ShaderSource++;
+			CopyInfo.DestUBTypeIndex = MetalPackedTypeNameToTypeIndex(CopyInfo.DestUBTypeName);
+			verify(Match(ShaderSource, ':'));
+
+			CopyInfo.DestOffsetInFloats = ParseNumber(ShaderSource);
+			verify(Match(ShaderSource, ':'));
+
+			CopyInfo.SizeInFloats = ParseNumber(ShaderSource);
+
+			Header.UniformBuffersCopyInfo.Add(CopyInfo);
+
+			uint16& Size = PackedGlobalArraySize.FindOrAdd(CopyInfo.DestUBTypeName);
+			Size = FMath::Max<uint16>(BytesPerComponent * (CopyInfo.DestOffsetInFloats + CopyInfo.SizeInFloats), Size);
+
+			if (Match(ShaderSource, '\n'))
+			{
+				break;
+			}
+
+			verify(Match(ShaderSource, ','));
+		}
+	}
+	Header.Bindings.bHasRegularUniformBuffers = bHasRegularUniformBuffers;
+
+	// Setup Packed Array info
+	Header.Bindings.PackedGlobalArrays.Reserve(PackedGlobalArraySize.Num());
+	for (auto Iterator = PackedGlobalArraySize.CreateIterator(); Iterator; ++Iterator)
+	{
+		ANSICHAR TypeName = Iterator.Key();
+		uint16 Size = Iterator.Value();
+		Size = (Size + 0xf) & (~0xf);
+		FMetalPackedArrayInfo Info;
+		Info.Size = Size;
+		Info.TypeName = TypeName;
+		Info.TypeIndex = MetalPackedTypeNameToTypeIndex(TypeName);
+		Header.Bindings.PackedGlobalArrays.Add(Info);
+	}
+
+	// Setup Packed Uniform Buffers info
+	Header.Bindings.PackedUniformBuffers.Reserve(PackedUniformBuffersSize.Num());
+	for (auto Iterator = PackedUniformBuffersSize.CreateIterator(); Iterator; ++Iterator)
+	{
+		int BufferIndex = Iterator.Key();
+		auto& ArraySizes = Iterator.Value();
+		TArray<FMetalPackedArrayInfo> InfoArray;
+		InfoArray.Reserve(ArraySizes.Num());
+		for (auto IterSizes = ArraySizes.CreateIterator(); IterSizes; ++IterSizes)
+		{
+			ANSICHAR TypeName = IterSizes.Key();
+			uint16 Size = IterSizes.Value();
+			Size = (Size + 0xf) & (~0xf);
+			FMetalPackedArrayInfo Info;
+			Info.Size = Size;
+			Info.TypeName = TypeName;
+			Info.TypeIndex = MetalPackedTypeNameToTypeIndex(TypeName);
+			InfoArray.Add(Info);
+		}
+
+		Header.Bindings.PackedUniformBuffers.Add(InfoArray);
+	}
+
+	// Then samplers.
+	if (FCStringAnsi::Strncmp(ShaderSource, SamplersPrefix, SamplersPrefixLen) == 0)
+	{
+		ShaderSource += SamplersPrefixLen;
+
+		while (*ShaderSource && *ShaderSource != '\n')
+		{
+			uint16 Offset = 0;
+			uint16 NumSamplers = 0;
+
+			FString ParameterName = ParseIdentifier(ShaderSource);
+			verify(ParameterName.Len() > 0);
+			verify(Match(ShaderSource, '('));
+			Offset = ParseNumber(ShaderSource);
+			verify(Match(ShaderSource, ':'));
+			NumSamplers = ParseNumber(ShaderSource);
+			ParameterMap.AddParameterAllocation(
+				*ParameterName,
+				0,
+				Offset,
+				NumSamplers
+				);
+
+			Header.Bindings.NumSamplers = FMath::Max<uint8>(
+				Header.Bindings.NumSamplers,
+				Offset + NumSamplers
+				);
+
+			if (Match(ShaderSource, '['))
+			{
+				// Sampler States
+				do
+				{
+					FString SamplerState = ParseIdentifier(ShaderSource);
+					checkSlow(SamplerState.Len() != 0);
+					ParameterMap.AddParameterAllocation(
+						*SamplerState,
+						0,
+						Offset,
+						NumSamplers
+						);
+				}
+				while (Match(ShaderSource, ','));
+				verify(Match(ShaderSource, ']'));
+			}
+
+			verify(Match(ShaderSource, ')'));
+
+			if (Match(ShaderSource, '\n'))
+			{
+				break;
+			}
+
+			// Skip the comma.
+			verify(Match(ShaderSource, ','));
+		}
+	}	
+
+	// Then UAVs (images in Metal)
+	if (FCStringAnsi::Strncmp(ShaderSource, UAVsPrefix, UAVsPrefixLen) == 0)
+	{
+		ShaderSource += UAVsPrefixLen;
+
+		while (*ShaderSource && *ShaderSource != '\n')
+		{
+			uint16 Offset = 0;
+			uint16 NumUAVs = 0;
+
+			FString ParameterName = ParseIdentifier(ShaderSource);
+			verify(ParameterName.Len() > 0);
+			verify(Match(ShaderSource, '('));
+			Offset = ParseNumber(ShaderSource);
+			verify(Match(ShaderSource, ':'));
+			NumUAVs = ParseNumber(ShaderSource);
+
+			ParameterMap.AddParameterAllocation(
+				*ParameterName,
+				0,
+				Offset,
+				NumUAVs
+				);
+
+			Header.Bindings.NumUAVs = FMath::Max<uint8>(
+				Header.Bindings.NumUAVs,
+				Offset + NumUAVs
+				);
+
+			verify(Match(ShaderSource, ')'));
+
+			if (Match(ShaderSource, '\n'))
+			{
+				break;
+			}
+
+			// Skip the comma.
+			verify(Match(ShaderSource, ','));
+		}
+	}
+
+	const int32 MaxSamplers = GetFeatureLevelMaxTextureSamplers(GetMaxSupportedFeatureLevel((EShaderPlatform)ShaderOutput.Target.Platform));
+
+	if (Header.Bindings.NumSamplers > MaxSamplers)
+	{
+		ShaderOutput.bSucceeded = false;
+		FShaderCompilerError* NewError = new(ShaderOutput.Errors) FShaderCompilerError();
+		NewError->StrippedErrorMessage =
+			FString::Printf(TEXT("shader uses %d samplers exceeding the limit of %d"),
+				Header.Bindings.NumSamplers, MaxSamplers);
+	}
+	else
+	{
+#if METAL_OFFLINE_COMPILE
+		// at this point, the shader source is ready to be compiled
+		FString InputFilename = FPaths::CreateTempFilename(*FPaths::EngineIntermediateDir(), TEXT("ShaderIn"), TEXT(""));
+		FString ObjFilename = InputFilename + TEXT(".o");
+		FString ArFilename = InputFilename + TEXT(".ar");
+		FString OutputFilename = InputFilename + TEXT(".lib");
+		InputFilename = InputFilename + TEXT(".metal");
+		
+		// write out shader source
+		FFileHelper::SaveStringToFile(FString(ShaderSource), *InputFilename);
+		
+		int32 ReturnCode = 0;
+		FString Results;
+		FString Errors;
+		bool bHadError = true;
+
+		// metal commandlines
+		FString Params = FString::Printf(TEXT("-std=ios-metal1.0 %s -o %s"), *InputFilename, *ObjFilename);
+		FPlatformProcess::ExecProcess( TEXT("/Applications/Xcode.app/Contents/Developer/Platforms/iPhoneOS.platform/usr/local/bin/metal"), *Params, &ReturnCode, &Results, &Errors );
+
+		// handle compile error
+		if (ReturnCode != 0 || IFileManager::Get().FileSize(*ObjFilename) <= 0)
+		{
+//			FShaderCompilerError* Error = new(OutErrors) FShaderCompilerError();
+//			Error->ErrorFile = InputFilename;
+//			Error->ErrorLineString = TEXT("0");
+//			Error->StrippedErrorMessage = Results + Errors;
+		}
+		else
+		{
+			Params = FString::Printf(TEXT("r %s %s"), *ArFilename, *ObjFilename);
+			FPlatformProcess::ExecProcess( TEXT("/Applications/Xcode.app/Contents/Developer/Platforms/iPhoneOS.platform/usr/local/bin/metal-ar"), *Params, &ReturnCode, &Results, &Errors );
+
+			// handle compile error
+			if (ReturnCode != 0 || IFileManager::Get().FileSize(*ArFilename) <= 0)
+			{
+//				FShaderCompilerError* Error = new(OutErrors) FShaderCompilerError();
+//				Error->ErrorFile = InputFilename;
+//				Error->ErrorLineString = TEXT("0");
+//				Error->StrippedErrorMessage = Results + Errors;
+			}
+			else
+			{
+				Params = FString::Printf(TEXT("-o %s %s"), *OutputFilename, *ArFilename);
+				FPlatformProcess::ExecProcess( TEXT("/Applications/Xcode.app/Contents/Developer/Platforms/iPhoneOS.platform/usr/local/bin/metallib"), *Params, &ReturnCode, &Results, &Errors );
+		
+				// handle compile error
+				if (ReturnCode != 0 || IFileManager::Get().FileSize(*OutputFilename) <= 0)
+				{
+//					FShaderCompilerError* Error = new(OutErrors) FShaderCompilerError();
+//					Error->ErrorFile = InputFilename;
+//					Error->ErrorLineString = TEXT("0");
+//					Error->StrippedErrorMessage = Results + Errors;
+				}
+				else
+				{
+					bHadError = false;
+					
+					// Write out the header and compiled shader code
+					FMemoryWriter Ar(ShaderOutput.Code, true);
+					uint8 PrecompiledFlag = 1;
+					Ar << PrecompiledFlag;
+					Ar << Header;
+
+					// load output
+					TArray<uint8> CompiledShader;
+					FFileHelper::LoadFileToArray(CompiledShader, *OutputFilename);
+					
+					// jam it into the output bytes
+					Ar.Serialize(CompiledShader.GetData(), CompiledShader.Num());
+					
+					ShaderOutput.NumInstructions = 0;
+					ShaderOutput.NumTextureSamplers = Header.Bindings.NumSamplers;
+					ShaderOutput.bSucceeded = true;
+				}
+			}
+		}
+		
+		if (bHadError)
+		{
+			// Write out the header and shader source code.
+			FMemoryWriter Ar(ShaderOutput.Code, true);
+			uint8 PrecompiledFlag = 0;
+			Ar << PrecompiledFlag;
+			Ar << Header;
+			Ar.Serialize((void*)ShaderSource, SourceLen + 1 - (ShaderSource - InShaderSource));
+			
+			ShaderOutput.NumInstructions = 0;
+			ShaderOutput.NumTextureSamplers = Header.Bindings.NumSamplers;
+			ShaderOutput.bSucceeded = true;
+		}
+		
+		IFileManager::Get().Delete(*InputFilename);
+		IFileManager::Get().Delete(*ObjFilename);
+		IFileManager::Get().Delete(*ArFilename);
+		IFileManager::Get().Delete(*OutputFilename);
+#else
+		// Write out the header and shader source code.
+		FMemoryWriter Ar(ShaderOutput.Code, true);
+		uint8 PrecompiledFlag = 0;
+		Ar << PrecompiledFlag;
+		Ar << Header;
+		Ar.Serialize((void*)ShaderSource, SourceLen + 1 - (ShaderSource - InShaderSource));
+		
+		ShaderOutput.NumInstructions = 0;
+		ShaderOutput.NumTextureSamplers = Header.Bindings.NumSamplers;
+		ShaderOutput.bSucceeded = true;
+#endif
+	}
+}
+
+/**
+ * Parse an error emitted by the HLSL cross-compiler.
+ * @param OutErrors - Array into which compiler errors may be added.
+ * @param InLine - A line from the compile log.
+ */
+static void ParseHlslccError(TArray<FShaderCompilerError>& OutErrors, const FString& InLine)
+{
+	const TCHAR* p = *InLine;
+	FShaderCompilerError* Error = new(OutErrors) FShaderCompilerError();
+
+	// Copy the filename.
+	while (*p && *p != TEXT('(')) { Error->ErrorFile += (*p++); }
+	Error->ErrorFile = GetRelativeShaderFilename(Error->ErrorFile);
+	p++;
+
+	// Parse the line number.
+	int32 LineNumber = 0;
+	while (*p && *p >= TEXT('0') && *p <= TEXT('9'))
+	{
+		LineNumber = 10 * LineNumber + (*p++ - TEXT('0'));
+	}
+	Error->ErrorLineString = *FString::Printf(TEXT("%d"), LineNumber);
+
+	// Skip to the warning message.
+	while (*p && (*p == TEXT(')') || *p == TEXT(':') || *p == TEXT(' ') || *p == TEXT('\t'))) { p++; }
+	Error->StrippedErrorMessage = p;
+}
+
+/*------------------------------------------------------------------------------
+	External interface.
+------------------------------------------------------------------------------*/
+
+static FString CreateCommandLineHLSLCC( const FString& ShaderFile, const FString& OutputFile, const FString& EntryPoint, EHlslShaderFrequency Frequency, uint32 CCFlags ) 
+{
+	const TCHAR* FrequencySwitch = TEXT("");
+	switch (Frequency)
+	{
+		case HSF_PixelShader:
+			FrequencySwitch = TEXT(" -ps");
+			break;
+
+		case HSF_VertexShader:
+			FrequencySwitch = TEXT(" -vs");
+			break;
+
+		case HSF_HullShader:
+			FrequencySwitch = TEXT(" -hs");
+			break;
+
+		case HSF_DomainShader:
+			FrequencySwitch = TEXT(" -ds");
+			break;
+		case HSF_ComputeShader:
+			FrequencySwitch = TEXT(" -cs");
+			break;
+
+		case HSF_GeometryShader:
+			FrequencySwitch = TEXT(" -gs");
+			break;
+
+		default:
+			check(0);
+	}
+
+	const TCHAR* VersionSwitch = TEXT("-metal");
+	const TCHAR* FlattenUB = (CCFlags & HLSLCC_FlattenUniformBuffers) ? TEXT("-flattenub") : TEXT("");
+
+	FString CmdLine = FPaths::RootDir() / FString::Printf(TEXT("Engine\\Source\\ThirdParty\\hlslcc\\hlslcc\\bin\\Win64\\VS2013\\hlslcc_64.exe %s -o=%s %s -entry=%s %s %s"), *ShaderFile, *OutputFile, FrequencySwitch, *EntryPoint, VersionSwitch, FlattenUB);
+	CmdLine += "\npause";
+	return CmdLine;
+}
+
+void CompileShader_Metal(const FShaderCompilerInput& Input,FShaderCompilerOutput& Output,const FString& WorkingDirectory)
+{
+	FString PreprocessedShader;
+	FShaderCompilerDefinitions AdditionalDefines;
+	EHlslCompileTarget HlslCompilerTarget = HCT_FeatureLevelES3_1;
+	AdditionalDefines.SetDefine(TEXT("IOS"), 1);
+	AdditionalDefines.SetDefine(TEXT("METAL_PROFILE"), 1);
+	AdditionalDefines.SetDefine(TEXT("ES2_PROFILE"), 1);
+	AdditionalDefines.SetDefine(TEXT("row_major"), TEXT(""));
+	
+	const bool bDumpDebugInfo = (Input.DumpDebugInfoPath != TEXT("") && IFileManager::Get().DirectoryExists(*Input.DumpDebugInfoPath));
+
+	AdditionalDefines.SetDefine(TEXT("COMPILER_SUPPORTS_ATTRIBUTES"), (uint32)1);
+	if (PreprocessShader(PreprocessedShader, Output, Input, AdditionalDefines))
+	{
+		char* MetalShaderSource = NULL;
+		char* ErrorLog = NULL;
+
+		const EHlslShaderFrequency FrequencyTable[] =
+		{
+			HSF_VertexShader,
+			HSF_InvalidFrequency,
+			HSF_InvalidFrequency,
+			HSF_PixelShader,
+			HSF_InvalidFrequency,
+			HSF_ComputeShader
+		};
+
+		const EHlslShaderFrequency Frequency = FrequencyTable[Input.Target.Frequency];
+		if (Frequency == HSF_InvalidFrequency)
+		{
+			Output.bSucceeded = false;
+			FShaderCompilerError* NewError = new(Output.Errors) FShaderCompilerError();
+			NewError->StrippedErrorMessage = FString::Printf(
+				TEXT("%s shaders not supported for use in Metal."),
+				GLFrequencyStringTable[Input.Target.Frequency]
+				);
+			return;
+		}
+
+		// Write out the preprocessed file and a batch file to compile it if requested (DumpDebugInfoPath is valid)
+		if (bDumpDebugInfo)
+		{
+			FArchive* FileWriter = IFileManager::Get().CreateFileWriter(*(Input.DumpDebugInfoPath / Input.SourceFilename + TEXT(".usf")));
+			if (FileWriter)
+			{
+				auto AnsiSourceFile = StringCast<ANSICHAR>(*PreprocessedShader);
+				FileWriter->Serialize((ANSICHAR*)AnsiSourceFile.Get(), AnsiSourceFile.Length());
+				FileWriter->Close();
+				delete FileWriter;
+			}
+		}
+
+		uint32 CCFlags = HLSLCC_NoPreprocess | HLSLCC_PackUniforms;
+		//CCFlags |= HLSLCC_FlattenUniformBuffers | HLSLCC_FlattenUniformBufferStructures;
+
+		if (bDumpDebugInfo)
+		{
+			const FString MetalFile = (Input.DumpDebugInfoPath / TEXT("Output.metal"));
+			const FString USFFile = (Input.DumpDebugInfoPath / Input.SourceFilename) + TEXT(".usf");
+			const FString CCBatchFileContents = CreateCommandLineHLSLCC(USFFile, MetalFile, *Input.EntryPointName, Frequency, CCFlags);
+			if (!CCBatchFileContents.IsEmpty())
+			{
+				FFileHelper::SaveStringToFile(CCBatchFileContents, *(Input.DumpDebugInfoPath / TEXT("CrossCompile.bat")));
+			}
+		}
+
+		FMetalCodeBackend MetalBackEnd(CCFlags);
+		FMetalLanguageSpec MetalLanguageSpec;
+		int32 Result = HlslCrossCompile(
+			TCHAR_TO_ANSI(*Input.SourceFilename),
+			TCHAR_TO_ANSI(*PreprocessedShader),
+			TCHAR_TO_ANSI(*Input.EntryPointName),
+			Frequency,
+			&MetalBackEnd,
+			&MetalLanguageSpec,
+			CCFlags,
+			HlslCompilerTarget,
+			&MetalShaderSource,
+			&ErrorLog
+			);
+
+		if (bDumpDebugInfo)
+		{
+			int32 MetalSourceLen = FCStringAnsi::Strlen(MetalShaderSource);
+			if (MetalSourceLen > 0)
+			{
+				FArchive* FileWriter = IFileManager::Get().CreateFileWriter(*(Input.DumpDebugInfoPath / Input.SourceFilename + TEXT(".metal")));
+				if (FileWriter)
+				{
+					FileWriter->Serialize(MetalShaderSource,MetalSourceLen+1);
+					FileWriter->Close();
+					delete FileWriter;
+				}
+			}
+		}
+
+		if (Result != 0)
+		{
+			int32 SourceLen = FCStringAnsi::Strlen(MetalShaderSource);
+			Output.Target = Input.Target;
+			BuildMetalShaderOutput(Output, MetalShaderSource, SourceLen, Output.Errors);
+		}
+		else
+		{
+			FString Tmp = ANSI_TO_TCHAR(ErrorLog);
+			TArray<FString> ErrorLines;
+			Tmp.ParseIntoArray(&ErrorLines, TEXT("\n"), true);
+			for (int32 LineIndex = 0; LineIndex < ErrorLines.Num(); ++LineIndex)
+			{
+				const FString& Line = ErrorLines[LineIndex];
+				ParseHlslccError(Output.Errors, Line);
+			}
+		}
+
+		if (MetalShaderSource)
+		{
+			free(MetalShaderSource);
+		}
+		if (ErrorLog)
+		{
+			free(ErrorLog);
+		}
+	}
+}
diff --git a/Engine/Source/Developer/MetalShaderFormat/Private/MetalShaderFormat.cpp b/Engine/Source/Developer/MetalShaderFormat/Private/MetalShaderFormat.cpp
new file mode 100644
--- /dev/null
+++ b/Engine/Source/Developer/MetalShaderFormat/Private/MetalShaderFormat.cpp
@@ -0,0 +1,84 @@
+// Copyright 1998-2014 Epic Games, Inc. All Rights Reserved.
+
+#include "UrbanShaderFormat.h"
+#include "Core.h"
+#include "ModuleInterface.h"
+#include "ModuleManager.h"
+#include "TargetPlatform.h"   
+#include "ShaderCore.h"
+
+
+static FName NAME_SF_METAL(TEXT("SF_METAL"));
+
+class FMetalShaderFormat : public IShaderFormat
+{
+	enum
+	{
+		HEADER_VERSION = 2,
+	};
+	struct FVersion
+	{
+		uint16 HLSLCCMajor		: 3;
+		uint16 HLSLCCMinor		: 7;
+		uint16 Format			: 5;
+		uint16 OfflineCompiled	: 1;
+	};
+public:
+	virtual uint16 GetVersion(FName Format) const OVERRIDE
+	{
+		static_assert(sizeof(FVersion) == sizeof(uint16), "Out of bits!");
+		union
+		{
+			FVersion Version;
+			uint16 Raw;
+		} Version;
+
+		Version.Version.Format = HEADER_VERSION;
+		Version.Version.HLSLCCMajor = HLSLCC_VersionMajor;
+		Version.Version.HLSLCCMinor = HLSLCC_VersionMinor;
+		Version.Version.OfflineCompiled = METAL_OFFLINE_COMPILE;
+		
+		// Check that we didn't overwrite any bits
+		check(Version.Version.Format == HEADER_VERSION);
+		check(Version.Version.HLSLCCMajor == HLSLCC_VersionMajor);
+		check(Version.Version.HLSLCCMinor == HLSLCC_VersionMinor);
+		check(Version.Version.OfflineCompiled == METAL_OFFLINE_COMPILE);
+
+		return Version.Raw;
+	}
+	virtual void GetSupportedFormats(TArray<FName>& OutFormats) const
+	{
+		OutFormats.Add(NAME_SF_METAL);
+	}
+	virtual void CompileShader(FName Format, const struct FShaderCompilerInput& Input, struct FShaderCompilerOutput& Output,const FString& WorkingDirectory) const
+	{
+		check(Format == NAME_SF_METAL);
+		CompileShader_Metal(Input, Output, WorkingDirectory);
+	}
+};
+
+/**
+ * Module for OpenGL shaders
+ */
+
+static IShaderFormat* Singleton = NULL;
+
+class FMetalShaderFormatModule : public IShaderFormatModule
+{
+public:
+	virtual ~FMetalShaderFormatModule()
+	{
+		delete Singleton;
+		Singleton = NULL;
+	}
+	virtual IShaderFormat* GetShaderFormat()
+	{
+		if (!Singleton)
+		{
+			Singleton = new FMetalShaderFormat();
+		}
+		return Singleton;
+	}
+};
+
+IMPLEMENT_MODULE( FMetalShaderFormatModule, UrbanShaderFormat);
diff --git a/Engine/Source/Developer/MetalShaderFormat/Private/UrbanShaderFormat.h b/Engine/Source/Developer/MetalShaderFormat/Private/UrbanShaderFormat.h
new file mode 100644
--- /dev/null
+++ b/Engine/Source/Developer/MetalShaderFormat/Private/UrbanShaderFormat.h
@@ -0,0 +1,10 @@
+// Copyright 1998-2014 Epic Games, Inc. All Rights Reserved.
+
+#pragma once
+
+extern void CompileShader_Metal(const struct FShaderCompilerInput& Input, struct FShaderCompilerOutput& Output, const class FString& WorkingDirectory);
+
+// Set this to 0 to get shader source in the graphics debugger
+// Note: Offline and runtime compiled shaders have separate DDC versions and can co-exist
+#define METAL_OFFLINE_COMPILE !PLATFORM_WINDOWS
+
