UE4 - Metal SRTs

diff --git a/Engine/Source/Developer/MetalShaderFormat/Private/MetalShaderCompiler.cpp b/Engine/Source/Developer/MetalShaderFormat/Private/MetalShaderCompiler.cpp
--- a/Engine/Source/Developer/MetalShaderFormat/Private/MetalShaderCompiler.cpp
+++ b/Engine/Source/Developer/MetalShaderFormat/Private/MetalShaderCompiler.cpp
@@ -1,7 +1,7 @@
 // Copyright 1998-2014 Epic Games, Inc. All Rights Reserved.
 ///
 
-#include "UrbanShaderFormat.h"
+#include "MetalShaderFormat.h"
 #include "Core.h"
 #include "../../../EmptyRHI/Public/EmptyShaderResources.h"
 
@@ -20,7 +20,7 @@
 #include "hlslcc.h"
 #include "Metal/MetalBackend.h"
 
-DEFINE_LOG_CATEGORY_STATIC(LogUrbanShaderCompiler, Log, All); 
+DEFINE_LOG_CATEGORY_STATIC(LogMetalShaderCompiler, Log, All); 
 
 /*------------------------------------------------------------------------------
 	Shader compiling.
@@ -595,6 +595,24 @@ static void BuildMetalShaderOutput(
 		}
 	}
 
+	// Build the SRT for this shader.
+	{
+		// Build the generic SRT for this shader.
+		FShaderResourceTable GenericSRT;
+		BuildResourceTableMapping(ShaderInput.Environment.ResourceTableMap, ShaderInput.Environment.ResourceTableLayoutHashes, UsedUniformBufferSlots, ShaderOutput.ParameterMap, GenericSRT);
+
+		// Copy over the bits indicating which resource tables are active.
+		Header.Bindings.ShaderResourceTable.ResourceTableBits = GenericSRT.ResourceTableBits;
+
+		Header.Bindings.ShaderResourceTable.ResourceTableLayoutHashes = GenericSRT.ResourceTableLayoutHashes;
+
+		// Now build our token streams.
+		BuildResourceTableTokenStream(GenericSRT.TextureMap, GenericSRT.MaxBoundResourceTable, Header.Bindings.ShaderResourceTable.TextureMap);
+		BuildResourceTableTokenStream(GenericSRT.ShaderResourceViewMap, GenericSRT.MaxBoundResourceTable, Header.Bindings.ShaderResourceTable.ShaderResourceViewMap);
+		BuildResourceTableTokenStream(GenericSRT.SamplerMap, GenericSRT.MaxBoundResourceTable, Header.Bindings.ShaderResourceTable.SamplerMap);
+		BuildResourceTableTokenStream(GenericSRT.UnorderedAccessViewMap, GenericSRT.MaxBoundResourceTable, Header.Bindings.ShaderResourceTable.UnorderedAccessViewMap);
+	}
+
 	const int32 MaxSamplers = GetFeatureLevelMaxTextureSamplers(GetMaxSupportedFeatureLevel((EShaderPlatform)ShaderOutput.Target.Platform));
 
 	if (Header.Bindings.NumSamplers > MaxSamplers)
diff --git a/Engine/Source/Developer/MetalShaderFormat/Private/MetalShaderFormat.cpp b/Engine/Source/Developer/MetalShaderFormat/Private/MetalShaderFormat.cpp
--- a/Engine/Source/Developer/MetalShaderFormat/Private/MetalShaderFormat.cpp
+++ b/Engine/Source/Developer/MetalShaderFormat/Private/MetalShaderFormat.cpp
@@ -1,6 +1,6 @@
 // Copyright 1998-2014 Epic Games, Inc. All Rights Reserved.
 
-#include "UrbanShaderFormat.h"
+#include "MetalShaderFormat.h"
 #include "Core.h"
 #include "ModuleInterface.h"
 #include "ModuleManager.h"
@@ -14,7 +14,7 @@ class FMetalShaderFormat : public IShaderFormat
 {
 	enum
 	{
-		HEADER_VERSION = 2,
+		HEADER_VERSION = 3,
 	};
 	struct FVersion
 	{
@@ -81,4 +81,4 @@ public:
 	}
 };
 
-IMPLEMENT_MODULE( FMetalShaderFormatModule, UrbanShaderFormat);
+IMPLEMENT_MODULE( FMetalShaderFormatModule, MetalShaderFormat);
diff --git a/Engine/Source/Developer/MetalShaderFormat/Private/MetalShaderFormat.h b/Engine/Source/Developer/MetalShaderFormat/Private/MetalShaderFormat.h
new file mode 100644
--- /dev/null
+++ b/Engine/Source/Developer/MetalShaderFormat/Private/MetalShaderFormat.h
@@ -0,0 +1,10 @@
+// Copyright 1998-2014 Epic Games, Inc. All Rights Reserved.
+
+#pragma once
+
+extern void CompileShader_Metal(const struct FShaderCompilerInput& Input, struct FShaderCompilerOutput& Output, const class FString& WorkingDirectory);
+
+// Set this to 0 to get shader source in the graphics debugger
+// Note: Offline and runtime compiled shaders have separate DDC versions and can co-exist
+#define METAL_OFFLINE_COMPILE !PLATFORM_WINDOWS
+
diff --git a/Engine/Source/Developer/MetalShaderFormat/Private/UrbanShaderFormat.h b/Engine/Source/Developer/MetalShaderFormat/Private/UrbanShaderFormat.h
deleted file mode 100644
